[SequenceNavigator] Add new CinematicSequenceNavigator plugin that implements its provider for the navigator.

diff --git a/Engine/Plugins/VirtualProduction/CinematicSequenceNavigator/CinematicSequenceNavigator.uplugin b/Engine/Plugins/VirtualProduction/CinematicSequenceNavigator/CinematicSequenceNavigator.uplugin
new file mode 100644
--- /dev/null
+++ b/Engine/Plugins/VirtualProduction/CinematicSequenceNavigator/CinematicSequenceNavigator.uplugin
@@ -0,0 +1,41 @@
+﻿{
+	"FileVersion": 3,
+	"Version": 1,
+	"VersionName": "1.0",
+	"FriendlyName": "Cinematic Sequence Navigator Bridge",
+	"Description": "Sequence Navigator Bridge for Cinematic Assemblies",
+	"Category": "Animation",
+	"CreatedBy": "Epic Games, Inc.",
+	"CreatedByURL": "https://epicgames.com",
+	"DocsURL": "",
+	"MarketplaceURL": "",
+	"CanContainContent": false,
+	"IsBetaVersion": false,
+	"IsExperimentalVersion": true,
+	"Installed": false,
+	"EnabledByDefault": false,
+	"SupportedTargetPlatforms": [
+		"Win64",
+		"Linux",
+		"Mac"
+	],
+	"Modules": [
+		{
+			"Name": "CinematicSequenceNavigator",
+			"Type": "Runtime",
+			"LoadingPhase": "Default"
+		}
+	],
+	"Plugins": [
+		{
+			"Name": "CinematicAssemblyTools",
+			"Enabled": true,
+			"Optional": true
+		},
+		{
+			"Name": "SequenceNavigator",
+			"Enabled": true,
+			"Optional": true
+		}
+	]
+}
\ No newline at end of file
diff --git a/Engine/Plugins/VirtualProduction/CinematicSequenceNavigator/Source/CinematicSequenceNavigator/CinematicSequenceNavigator.Build.cs b/Engine/Plugins/VirtualProduction/CinematicSequenceNavigator/Source/CinematicSequenceNavigator/CinematicSequenceNavigator.Build.cs
new file mode 100644
--- /dev/null
+++ b/Engine/Plugins/VirtualProduction/CinematicSequenceNavigator/Source/CinematicSequenceNavigator/CinematicSequenceNavigator.Build.cs
@@ -0,0 +1,28 @@
+﻿// Copyright Epic Games, Inc. All Rights Reserved.
+
+using UnrealBuildTool;
+
+public class CinematicSequenceNavigator : ModuleRules
+{
+	public CinematicSequenceNavigator(ReadOnlyTargetRules Target) : base(Target)
+	{
+		PCHUsage = ModuleRules.PCHUsageMode.UseExplicitOrSharedPCHs;
+		
+		PublicDependencyModuleNames.AddRange(new string[] {
+			"Core"
+		});
+
+		PrivateDependencyModuleNames.AddRange(new string[] {
+			"CineAssemblyTools",
+			"CoreUObject",
+			"Engine",
+			"MovieScene",
+			"Sequencer",
+			"SequenceNavigator",
+			"Slate",
+			"SlateCore",
+			"ToolMenus",
+			"UnrealEd"
+		});
+	}
+}
diff --git a/Engine/Plugins/VirtualProduction/CinematicSequenceNavigator/Source/CinematicSequenceNavigator/Private/CinematicNavigationToolProvider.cpp b/Engine/Plugins/VirtualProduction/CinematicSequenceNavigator/Source/CinematicSequenceNavigator/Private/CinematicNavigationToolProvider.cpp
new file mode 100644
--- /dev/null
+++ b/Engine/Plugins/VirtualProduction/CinematicSequenceNavigator/Source/CinematicSequenceNavigator/Private/CinematicNavigationToolProvider.cpp
@@ -0,0 +1,284 @@
+﻿// Copyright Epic Games, Inc. All Rights Reserved.
+
+#include "CinematicNavigationToolProvider.h"
+#include "CineAssembly.h"
+#include "CinematicNavigationToolProvider.h"
+#include "CinematicSequenceNavigatorCommands.h"
+#include "Columns/NavigationToolColorColumn.h"
+#include "Columns/NavigationToolColumnExtender.h"
+#include "Columns/NavigationToolCommentColumn.h"
+#include "Columns/NavigationToolDeactiveStateColumn.h"
+#include "Columns/NavigationToolHBiasColumn.h"
+#include "Columns/NavigationToolInTimeColumn.h"
+#include "Columns/NavigationToolItemsColumn.h"
+#include "Columns/NavigationToolLabelColumn.h"
+#include "Columns/NavigationToolLengthColumn.h"
+#include "Columns/NavigationToolLockColumn.h"
+#include "Columns/NavigationToolMarkerVisibilityColumn.h"
+#include "Columns/NavigationToolOutTimeColumn.h"
+#include "Columns/NavigationToolPlayheadColumn.h"
+#include "Columns/NavigationToolRevisionControlColumn.h"
+#include "Columns/NavigationToolStartFrameOffsetColumn.h"
+#include "Columns/NavigationToolTakeColumn.h"
+#include "Engine/Level.h"
+#include "Filters/Filters/NavigationToolBuiltInFilterParams.h"
+#include "INavigationTool.h"
+#include "ISequencer.h"
+#include "ISettingsModule.h"
+#include "ISourceControlModule.h"
+#include "ItemActions/NavigationToolAddItem.h"
+#include "Items/NavigationToolItemParameters.h"
+#include "Items/NavigationToolSequence.h"
+#include "MovieScene.h"
+#include "NavigationToolSettings.h"
+
+#define LOCTEXT_NAMESPACE "CinematicNavigationToolProvider"
+
+namespace UE::CineAssemblyTools
+{
+
+const FName FCinematicNavigationToolProvider::Identifier = TEXT("CinematicAssembly");
+
+const FText FCinematicNavigationToolProvider::CinematicColumnViewName = LOCTEXT("CinematicAssemblyColumnViewName", "Cinematic Assembly");
+
+const FName FCinematicNavigationToolProvider::ToolbarSectionName = TEXT("CinematicAssembly");
+const FName FCinematicNavigationToolProvider::ContextMenuSectionName = TEXT("CinematicAssemblyActions");
+
+FCinematicNavigationToolProvider::FCinematicNavigationToolProvider(const TSharedRef<ISequencer>& InSequencer)
+	: WeakSequencer(InSequencer)
+	, ToolCommands(MakeShared<FUICommandList>())
+{
+}
+
+FName FCinematicNavigationToolProvider::GetIdentifier() const
+{
+	return Identifier;
+}
+
+TSet<TSubclassOf<UMovieSceneSequence>> FCinematicNavigationToolProvider::GetSupportedSequenceClasses() const
+{
+	return { UCineAssembly::StaticClass() };
+}
+
+FText FCinematicNavigationToolProvider::GetDefaultColumnView() const
+{
+	return CinematicColumnViewName;
+}
+
+FMovieSceneEditorData* FCinematicNavigationToolProvider::GetRootMovieSceneEditorData(const UE::SequenceNavigator::INavigationTool& InTool) const
+{
+	const TSharedPtr<ISequencer> Sequencer = InTool.GetSequencer();
+	if (!Sequencer.IsValid())
+	{
+		return nullptr;
+	}
+
+	UMovieSceneSequence* const RootMovieSceneSequence = Sequencer->GetRootMovieSceneSequence();
+	if (!RootMovieSceneSequence)
+	{
+		return nullptr;
+	}
+
+	UMovieScene* const RootMovieScene = RootMovieSceneSequence->GetMovieScene();
+	if (!RootMovieScene)
+	{
+		return nullptr;
+	}
+
+	return &RootMovieScene->GetEditorData();
+}
+
+FNavigationToolSaveState* FCinematicNavigationToolProvider::GetSaveState(const UE::SequenceNavigator::INavigationTool& InTool) const
+{
+	if (FMovieSceneEditorData* const EditorData = GetRootMovieSceneEditorData(InTool))
+	{
+		return &EditorData->NavigationToolState;
+	}
+	return nullptr;
+}
+
+void FCinematicNavigationToolProvider::SetSaveState(const UE::SequenceNavigator::INavigationTool& InTool
+	, const FNavigationToolSaveState& InSaveState) const
+{
+	if (FMovieSceneEditorData* const EditorData = GetRootMovieSceneEditorData(InTool))
+	{
+		EditorData->NavigationToolState = InSaveState;
+	}
+}
+
+void FCinematicNavigationToolProvider::BindCommands(const TSharedRef<FUICommandList>& InCommandList)
+{
+	InCommandList->Append(ToolCommands);
+
+	//@TODO: Bind commands to functions
+}
+
+void FCinematicNavigationToolProvider::OnActivate()
+{
+	ExtendToolToolBar();
+	ExtendToolItemContextMenu();
+}
+
+void FCinematicNavigationToolProvider::OnDeactivate()
+{
+	RemoveToolToolBarExtension();
+	RemoveToolItemContextMenuExtension();
+}
+
+void FCinematicNavigationToolProvider::OnExtendColumns(UE::SequenceNavigator::FNavigationToolColumnExtender& OutExtender)
+{
+	using namespace UE::SequenceNavigator;
+
+	// Support built in columns
+	OutExtender.AddColumn<FNavigationToolPlayheadColumn>();
+	OutExtender.AddColumn<FNavigationToolDeactiveStateColumn>();
+	OutExtender.AddColumn<FNavigationToolMarkerVisibilityColumn>();
+	OutExtender.AddColumn<FNavigationToolLockColumn>();
+	OutExtender.AddColumn<FNavigationToolColorColumn>();
+	OutExtender.AddColumn<FNavigationToolLabelColumn>();
+	OutExtender.AddColumn<FNavigationToolItemsColumn>();
+	OutExtender.AddColumn<FNavigationToolInTimeColumn>();
+	OutExtender.AddColumn<FNavigationToolOutTimeColumn>();
+	OutExtender.AddColumn<FNavigationToolLengthColumn>();
+	OutExtender.AddColumn<FNavigationToolHBiasColumn>();
+	OutExtender.AddColumn<FNavigationToolStartFrameOffsetColumn>();
+	OutExtender.AddColumn<FNavigationToolTakeColumn>();
+	OutExtender.AddColumn<FNavigationToolCommentColumn>();
+
+	ISourceControlModule& SourceControlModule = ISourceControlModule::Get();
+	if (SourceControlModule.IsEnabled())
+	{
+		OutExtender.AddColumn<FNavigationToolRevisionControlColumn>();
+	}
+
+	//@TODO: Add Cinematic Assembly specific columns
+
+	FNavigationToolProvider::OnExtendColumns(OutExtender);
+}
+
+void FCinematicNavigationToolProvider::OnExtendColumnViews(TSet<FNavigationToolColumnView>& OutColumnViews)
+{
+	using namespace UE::SequenceNavigator;
+
+	FNavigationToolColumnView& CinematicColumnView = OutColumnViews
+		.FindOrAdd(FNavigationToolColumnView(CinematicColumnViewName));
+
+	CinematicColumnView.VisibleColumns =
+		{
+			FNavigationToolPlayheadColumn::StaticColumnId(),
+			FNavigationToolColorColumn::StaticColumnId(),
+			FNavigationToolLabelColumn::StaticColumnId(),
+			FNavigationToolItemsColumn::StaticColumnId(),
+			FNavigationToolInTimeColumn::StaticColumnId(),
+			FNavigationToolOutTimeColumn::StaticColumnId(),
+			FNavigationToolHBiasColumn::StaticColumnId()
+		};
+
+	FNavigationToolProvider::OnExtendColumnViews(OutColumnViews);
+}
+
+void FCinematicNavigationToolProvider::OnExtendItemChildren(UE::SequenceNavigator::INavigationTool& InTool
+	, const UE::SequenceNavigator::FNavigationToolItemPtr& InParentItem
+	, TArray<UE::SequenceNavigator::FNavigationToolItemPtr>& OutChildren
+	, const bool bInRecursive)
+{
+	using namespace UE::SequenceNavigator;
+
+	FNavigationToolProvider::OnExtendItemChildren(InTool, InParentItem, OutChildren, bInRecursive);
+
+	// Only extending root item
+	if (InParentItem->GetItemId() != FNavigationToolItemId::RootId)
+	{
+		return;
+	}
+
+	const TSharedPtr<ISequencer> Sequencer = InTool.GetSequencer();
+	if (!Sequencer.IsValid())
+	{
+		return;
+	}
+
+	UMovieSceneSequence* const RootSequence = Sequencer->GetRootMovieSceneSequence();
+	if (!RootSequence)
+	{
+		return;
+	}
+
+	const TSharedRef<FCinematicNavigationToolProvider> Provider = SharedThis(this);
+
+	const FNavigationToolItemPtr NewItem = InTool.FindOrAdd<FNavigationToolSequence>(Provider, InParentItem, RootSequence, nullptr, 0);
+	const FNavigationToolItemFlagGuard Guard(NewItem, ENavigationToolItemFlags::IgnorePendingKill);
+	OutChildren.Add(NewItem);
+	if (bInRecursive)
+	{
+		NewItem->FindChildren(OutChildren, bInRecursive);
+	}
+}
+
+void FCinematicNavigationToolProvider::OnExtendBuiltInFilters(TArray<FNavigationToolBuiltInFilterParams>& OutFilterParams)
+{
+	OutFilterParams.Add(FNavigationToolBuiltInFilterParams::CreateSequenceFilter());
+	OutFilterParams.Add(FNavigationToolBuiltInFilterParams::CreateTrackFilter());
+	OutFilterParams.Add(FNavigationToolBuiltInFilterParams::CreateBindingFilter());
+	OutFilterParams.Add(FNavigationToolBuiltInFilterParams::CreateMarkerFilter());
+
+	FNavigationToolProvider::OnExtendBuiltInFilters(OutFilterParams);
+}
+
+TSharedPtr<SequenceNavigator::INavigationTool> FCinematicNavigationToolProvider::GetNavigationTool() const
+{
+	if (const TSharedPtr<ISequencer> Sequencer = WeakSequencer.Pin())
+	{
+		return SequenceNavigator::FNavigationToolExtender::FindNavigationTool(Sequencer.ToSharedRef());
+	}
+	return nullptr;
+}
+
+void FCinematicNavigationToolProvider::ExtendToolToolBar()
+{
+	//@TODO: Add navigator toolbar buttons
+}
+
+void FCinematicNavigationToolProvider::ExtendToolItemContextMenu()
+{
+	UToolMenus* ToolMenus = UToolMenus::Get();
+	if (!ToolMenus)
+	{
+		return;
+	}
+
+	UToolMenu* const ExtendedMenu = ToolMenus->ExtendMenu(UE::SequenceNavigator::GetItemContextMenuName());
+	if (!ExtendedMenu)
+	{
+		return;
+	}
+
+	//const FCinematicSequenceNavigatorCommands& CinematicSequenceNavigatorCommands = FCinematicSequenceNavigatorCommands::Get();
+
+	//@TODO: Extend item context menu
+}
+
+void FCinematicNavigationToolProvider::RemoveToolToolBarExtension()
+{
+	if (UToolMenu* const Menu = UToolMenus::Get()->FindMenu(UE::SequenceNavigator::GetToolBarMenuName()))
+	{
+		Menu->RemoveSection(ToolbarSectionName);
+	}
+}
+
+void FCinematicNavigationToolProvider::RemoveToolItemContextMenuExtension()
+{
+	if (UToolMenu* const Menu = UToolMenus::Get()->FindMenu(UE::SequenceNavigator::GetItemContextMenuName()))
+	{
+		Menu->RemoveSection(ContextMenuSectionName);
+	}
+}
+
+void FCinematicNavigationToolProvider::ExtendToolItemDragDropOp(UE::SequenceNavigator::FNavigationToolItemDragDropOp& InDragDropOp)
+{
+	//@TODO: Extend drag and drop operations
+}
+
+} // namespace UE::CineAssemblyTools
+
+#undef LOCTEXT_NAMESPACE
diff --git a/Engine/Plugins/VirtualProduction/CinematicSequenceNavigator/Source/CinematicSequenceNavigator/Private/CinematicNavigationToolProvider.h b/Engine/Plugins/VirtualProduction/CinematicSequenceNavigator/Source/CinematicSequenceNavigator/Private/CinematicNavigationToolProvider.h
new file mode 100644
--- /dev/null
+++ b/Engine/Plugins/VirtualProduction/CinematicSequenceNavigator/Source/CinematicSequenceNavigator/Private/CinematicNavigationToolProvider.h
@@ -0,0 +1,85 @@
+﻿// Copyright Epic Games, Inc. All Rights Reserved.
+
+#pragma once
+
+#include "Delegates/IDelegateInstance.h"
+#include "NavigationToolDefines.h"
+#include "Providers/NavigationToolProvider.h"
+#include "Templates/SharedPointer.h"
+#include "UObject/WeakObjectPtr.h"
+
+class FUICommandList;
+class UMovieSceneSequence;
+class UObject;
+struct FNavigationToolSaveState;
+struct FMovieSceneEditorData;
+
+namespace UE::SequenceNavigator
+{
+	class FNavigationToolAvaSequence;
+	class FNavigationToolItem;
+	class FNavigationToolItemDragDropOp;
+	class INavigationTool;
+}
+
+namespace UE::CineAssemblyTools
+{
+
+class FCinematicNavigationToolProvider : public SequenceNavigator::FNavigationToolProvider
+{
+public:
+	static const FName Identifier;
+
+	static const FText CinematicColumnViewName;
+
+	FCinematicNavigationToolProvider(const TSharedRef<ISequencer>& InSequencer);
+	virtual ~FCinematicNavigationToolProvider() override = default;
+
+	//~ Begin INavigationToolProvider
+
+	virtual FName GetIdentifier() const override;
+
+	virtual TSet<TSubclassOf<UMovieSceneSequence>> GetSupportedSequenceClasses() const override;
+
+	virtual FText GetDefaultColumnView() const override;
+
+	virtual FNavigationToolSaveState* GetSaveState(const SequenceNavigator::INavigationTool& InTool) const override;
+	virtual void SetSaveState(const SequenceNavigator::INavigationTool& InTool, const FNavigationToolSaveState& InSaveState) const override;
+
+	virtual void BindCommands(const TSharedRef<FUICommandList>& InCommandList) override;
+
+	virtual void OnActivate() override;
+	virtual void OnDeactivate() override;
+
+	virtual void OnExtendColumns(SequenceNavigator::FNavigationToolColumnExtender& OutExtender) override;
+	virtual void OnExtendColumnViews(TSet<FNavigationToolColumnView>& OutColumnViews) override;
+	virtual void OnExtendItemChildren(SequenceNavigator::INavigationTool& InTool
+		, const SequenceNavigator::FNavigationToolItemPtr& InParentItem
+		, TArray<SequenceNavigator::FNavigationToolItemPtr>& OutChildren
+		, const bool bInRecursive) override;
+	virtual void OnExtendBuiltInFilters(TArray<FNavigationToolBuiltInFilterParams>& OutFilterParams) override;
+
+	//~ End INavigationToolProvider
+
+private:
+	static const FName ToolbarSectionName;
+	static const FName ContextMenuSectionName;
+
+	TSharedPtr<SequenceNavigator::INavigationTool> GetNavigationTool() const;
+
+	void ExtendToolToolBar();
+	void ExtendToolItemContextMenu();
+
+	void RemoveToolToolBarExtension();
+	void RemoveToolItemContextMenuExtension();
+
+	void ExtendToolItemDragDropOp(SequenceNavigator::FNavigationToolItemDragDropOp& InDragDropOp);
+
+	FMovieSceneEditorData* GetRootMovieSceneEditorData(const UE::SequenceNavigator::INavigationTool& InTool) const;
+
+	TWeakPtr<ISequencer> WeakSequencer;
+
+	TSharedRef<FUICommandList> ToolCommands;
+};
+
+} // namespace UE::CineAssemblyTools
diff --git a/Engine/Plugins/VirtualProduction/CinematicSequenceNavigator/Source/CinematicSequenceNavigator/Private/CinematicSequenceNavigatorCommands.cpp b/Engine/Plugins/VirtualProduction/CinematicSequenceNavigator/Source/CinematicSequenceNavigator/Private/CinematicSequenceNavigatorCommands.cpp
new file mode 100644
--- /dev/null
+++ b/Engine/Plugins/VirtualProduction/CinematicSequenceNavigator/Source/CinematicSequenceNavigator/Private/CinematicSequenceNavigatorCommands.cpp
@@ -0,0 +1,12 @@
+﻿// Copyright Epic Games, Inc. All Rights Reserved.
+
+#include "CinematicSequenceNavigatorCommands.h"
+
+#define LOCTEXT_NAMESPACE "CinematicSequenceNavigatorCommands"
+
+void FCinematicSequenceNavigatorCommands::RegisterCommands()
+{
+	// @TODO: Register commands here
+}
+
+#undef LOCTEXT_NAMESPACE
diff --git a/Engine/Plugins/VirtualProduction/CinematicSequenceNavigator/Source/CinematicSequenceNavigator/Private/CinematicSequenceNavigatorCommands.h b/Engine/Plugins/VirtualProduction/CinematicSequenceNavigator/Source/CinematicSequenceNavigator/Private/CinematicSequenceNavigatorCommands.h
new file mode 100644
--- /dev/null
+++ b/Engine/Plugins/VirtualProduction/CinematicSequenceNavigator/Source/CinematicSequenceNavigator/Private/CinematicSequenceNavigatorCommands.h
@@ -0,0 +1,23 @@
+﻿// Copyright Epic Games, Inc. All Rights Reserved.
+
+#pragma once
+
+#include "Framework/Commands/Commands.h"
+#include "Styling/AppStyle.h"
+
+class FCinematicSequenceNavigatorCommands : public TCommands<FCinematicSequenceNavigatorCommands>
+{
+public:
+	FCinematicSequenceNavigatorCommands()
+		: TCommands<FCinematicSequenceNavigatorCommands>(TEXT("CinematicSequenceNavigatorCommands")
+		, NSLOCTEXT("CinematicSequenceNavigatorCommands"
+			, "CinematicSequenceNavigatorCommands", "Cinematic Sequence Navigator Commands")
+		, NAME_None
+		, FAppStyle::GetAppStyleSetName())
+	{
+	}
+
+	virtual void RegisterCommands() override;
+
+	// @TODO: Add commands here
+};
diff --git a/Engine/Plugins/VirtualProduction/CinematicSequenceNavigator/Source/CinematicSequenceNavigator/Private/CinematicSequenceNavigatorModule.cpp b/Engine/Plugins/VirtualProduction/CinematicSequenceNavigator/Source/CinematicSequenceNavigator/Private/CinematicSequenceNavigatorModule.cpp
new file mode 100644
--- /dev/null
+++ b/Engine/Plugins/VirtualProduction/CinematicSequenceNavigator/Source/CinematicSequenceNavigator/Private/CinematicSequenceNavigatorModule.cpp
@@ -0,0 +1,78 @@
+﻿// Copyright Epic Games, Inc. All Rights Reserved.
+
+#include "CinematicSequenceNavigatorModule.h"
+#include "CinematicNavigationToolProvider.h"
+#include "CinematicSequenceNavigatorCommands.h"
+#include "ISequencer.h"
+#include "ISequencerModule.h"
+#include "Modules/ModuleManager.h"
+#include "NavigationToolExtender.h"
+
+#define LOCTEXT_NAMESPACE "FCinematicSequenceNavigatorModule"
+
+using namespace UE::CineAssemblyTools;
+using namespace UE::SequenceNavigator;
+
+void FCinematicSequenceNavigatorModule::StartupModule()
+{
+	FCinematicSequenceNavigatorCommands::Register();
+
+	if (ISequencerModule* const SequencerModule = FModuleManager::LoadModulePtr<ISequencerModule>(TEXT("Sequencer")))
+	{
+		SequencerCreatedHandle = SequencerModule->RegisterOnSequencerCreated(
+			FOnSequencerCreated::FDelegate::CreateRaw(this, &FCinematicSequenceNavigatorModule::OnSequencerCreated));
+	}
+}
+
+void FCinematicSequenceNavigatorModule::ShutdownModule()
+{
+	if (ISequencerModule* const SequencerModule = FModuleManager::LoadModulePtr<ISequencerModule>(TEXT("Sequencer")))
+	{
+		SequencerModule->UnregisterOnSequencerCreated(SequencerCreatedHandle);
+		SequencerCreatedHandle.Reset();
+	}
+
+	FCinematicSequenceNavigatorCommands::Unregister();
+}
+
+void FCinematicSequenceNavigatorModule::OnSequencerCreated(const TSharedRef<ISequencer> InSequencer)
+{
+	const FName ToolId = FNavigationToolExtender::GetToolInstanceId(*InSequencer);
+
+	const TSharedPtr<FNavigationToolProvider> FoundProvider = FNavigationToolExtender::FindToolProvider(ToolId
+		, FCinematicNavigationToolProvider::Identifier);
+	if (!FoundProvider.IsValid())
+	{
+		NavigationToolProvider = MakeShared<FCinematicNavigationToolProvider>(InSequencer);
+		if (NavigationToolProvider->IsSequenceSupported(InSequencer->GetRootMovieSceneSequence()))
+		{
+			FNavigationToolExtender::RegisterToolProvider(InSequencer, NavigationToolProvider.ToSharedRef());
+		}
+		else
+		{
+			NavigationToolProvider.Reset();
+		}
+	}
+
+	SequencerClosedHandle = InSequencer->OnCloseEvent().AddRaw(this, &FCinematicSequenceNavigatorModule::OnSequencerClosed);
+}
+
+void FCinematicSequenceNavigatorModule::OnSequencerClosed(const TSharedRef<ISequencer> InSequencer)
+{
+	if (!NavigationToolProvider.IsValid())
+	{
+		return;
+	}
+
+	const FName ToolId = FNavigationToolExtender::GetToolInstanceId(*InSequencer);
+
+	if (FNavigationToolExtender::UnregisterToolProvider(ToolId, NavigationToolProvider->GetIdentifier()))
+	{
+		NavigationToolProvider.Reset();
+		SequencerClosedHandle.Reset();
+	}
+}
+
+#undef LOCTEXT_NAMESPACE
+	
+IMPLEMENT_MODULE(FCinematicSequenceNavigatorModule, CinematicSequenceNavigator)
diff --git a/Engine/Plugins/VirtualProduction/CinematicSequenceNavigator/Source/CinematicSequenceNavigator/Private/CinematicSequenceNavigatorModule.h b/Engine/Plugins/VirtualProduction/CinematicSequenceNavigator/Source/CinematicSequenceNavigator/Private/CinematicSequenceNavigatorModule.h
new file mode 100644
--- /dev/null
+++ b/Engine/Plugins/VirtualProduction/CinematicSequenceNavigator/Source/CinematicSequenceNavigator/Private/CinematicSequenceNavigatorModule.h
@@ -0,0 +1,33 @@
+﻿// Copyright Epic Games, Inc. All Rights Reserved.
+
+#pragma once
+
+#include "Delegates/IDelegateInstance.h"
+#include "Modules/ModuleInterface.h"
+#include "Templates/SharedPointer.h"
+
+class ISequencer;
+
+namespace UE::CineAssemblyTools
+{
+	class FCinematicNavigationToolProvider;
+}
+
+class FCinematicSequenceNavigatorModule : public IModuleInterface
+{
+public:
+	//~ Begin IModuleInterface
+	virtual void StartupModule() override;
+	virtual void ShutdownModule() override;
+	//~ End IModuleInterface
+
+private:
+	void OnSequencerCreated(const TSharedRef<ISequencer> InSequencer);
+	void OnSequencerClosed(const TSharedRef<ISequencer> InSequencer);
+
+	FDelegateHandle SequencerCreatedHandle;
+	FDelegateHandle SequencerClosedHandle;
+
+	/** The provider that supplies data and extends the Navigation Tool */
+	TSharedPtr<UE::CineAssemblyTools::FCinematicNavigationToolProvider> NavigationToolProvider;
+};
